[FastGeo] Allow runtime use cases to opt out of automatic object reference collection to avoid cost of serializing FastGeo components.
[PCG] Manually collect object references when spawning FastGeo components.

Removes significant serialization cost, but places responsibility on user to manage references.

#rnx

diff --git a/Engine/Plugins/Experimental/FastGeoStreaming/Source/FastGeoStreaming/Internal/FastGeoContainer.h b/Engine/Plugins/Experimental/FastGeoStreaming/Source/FastGeoStreaming/Internal/FastGeoContainer.h
--- a/Engine/Plugins/Experimental/FastGeoStreaming/Source/FastGeoStreaming/Internal/FastGeoContainer.h
+++ b/Engine/Plugins/Experimental/FastGeoStreaming/Source/FastGeoStreaming/Internal/FastGeoContainer.h
@@ -127,7 +127,7 @@ public:
 	virtual IPhysicsBodyInstanceOwner* ResolvePhysicsBodyInstanceOwner(Chaos::FConstPhysicsObjectHandle PhysicsObject) override;
 	//~End IPhysicsBodyInstanceOwnerResolver
 
-	void OnCreated();
+	void OnCreated(bool bCollectReferences = true);
 	void AddComponentCluster(FFastGeoComponentCluster* ComponentCluster);
 
 	void InitializeDynamicProperties(bool bInitForPlay = false);
diff --git a/Engine/Plugins/Experimental/FastGeoStreaming/Source/FastGeoStreaming/Private/FastGeoContainer.cpp b/Engine/Plugins/Experimental/FastGeoStreaming/Source/FastGeoStreaming/Private/FastGeoContainer.cpp
--- a/Engine/Plugins/Experimental/FastGeoStreaming/Source/FastGeoStreaming/Private/FastGeoContainer.cpp
+++ b/Engine/Plugins/Experimental/FastGeoStreaming/Source/FastGeoStreaming/Private/FastGeoContainer.cpp
@@ -637,19 +637,25 @@ void UFastGeoContainer::CollectAssetReferences()
 	Assets = FFastGeoGatherFastGeoContainerAssetRefsArchive(*this).GetUniqueAssets().Array();
 }
 
-void UFastGeoContainer::OnCreated()
+void UFastGeoContainer::OnCreated(bool bCollectReferences)
 {
+	const bool bIsGameWorld = GetWorld()->IsGameWorld();
+
 	// Initialize component clusters & components dynamic properties
-	InitializeDynamicProperties(GetLevel()->GetWorld()->IsGameWorld());
+	InitializeDynamicProperties(bIsGameWorld);
 
 #if WITH_EDITOR
 	// In editor, we need a UObject representation of the components for some operations
 	CreateEditorProxyObjects();
 #endif
 
-	// Collect references in order to avoid garbage collection of objects that may now be unreferenced
-	// The fast geo container will hold onto those objects if necessary.
-	CollectAssetReferences();
+	// Always collect references outside of game worlds.
+	if (bCollectReferences || !bIsGameWorld)
+	{
+		// Collect references in order to avoid garbage collection of objects that may now be unreferenced
+		// The fast geo container will hold onto those objects if necessary.
+		CollectAssetReferences();
+	}
 }
 
 void UFastGeoContainer::Serialize(FArchive& Ar)
diff --git a/Engine/Plugins/Experimental/PCGInterops/PCGFastGeoInterop/Source/PCGFastGeoInterop/Private/Components/PCGManagedFastGeoContainer.h b/Engine/Plugins/Experimental/PCGInterops/PCGFastGeoInterop/Source/PCGFastGeoInterop/Private/Components/PCGManagedFastGeoContainer.h
--- a/Engine/Plugins/Experimental/PCGInterops/PCGFastGeoInterop/Source/PCGFastGeoInterop/Private/Components/PCGManagedFastGeoContainer.h
+++ b/Engine/Plugins/Experimental/PCGInterops/PCGFastGeoInterop/Source/PCGFastGeoInterop/Private/Components/PCGManagedFastGeoContainer.h
@@ -8,6 +8,8 @@
 
 #include "PCGManagedFastGeoContainer.generated.h"
 
+struct FPCGProceduralISMComponentDescriptor;
+
 UCLASS()
 class UPCGManagedFastGeoContainer : public UPCGManagedResource
 {
@@ -19,6 +21,16 @@ public:
 	virtual bool ReleaseOnTeardown() const override { return true; }
 	//~End UPCGManagedResource interface
 
+	void SetFastGeoContainer(UFastGeoContainer* InFastGeo) { FastGeo = InFastGeo; }
+
+	/** Set references to objects to be kept alive. */
+	void SetObjectReferences(TArray<TObjectPtr<UObject>>&& InObjectReferences) { ObjectReferences = MoveTemp(InObjectReferences); }
+
+protected:
 	UPROPERTY(Transient)
 	TObjectPtr<UFastGeoContainer> FastGeo;
+
+	/** Referenced objects to be kept alive. */
+	UPROPERTY(Transient)
+	TArray<TObjectPtr<UObject>> ObjectReferences;
 };
diff --git a/Engine/Plugins/Experimental/PCGInterops/PCGFastGeoInterop/Source/PCGFastGeoInterop/Private/Compute/PrimitiveFactories/PCGPrimitiveFactoryFastGeoPISMC.cpp b/Engine/Plugins/Experimental/PCGInterops/PCGFastGeoInterop/Source/PCGFastGeoInterop/Private/Compute/PrimitiveFactories/PCGPrimitiveFactoryFastGeoPISMC.cpp
--- a/Engine/Plugins/Experimental/PCGInterops/PCGFastGeoInterop/Source/PCGFastGeoInterop/Private/Compute/PrimitiveFactories/PCGPrimitiveFactoryFastGeoPISMC.cpp
+++ b/Engine/Plugins/Experimental/PCGInterops/PCGFastGeoInterop/Source/PCGFastGeoInterop/Private/Compute/PrimitiveFactories/PCGPrimitiveFactoryFastGeoPISMC.cpp
@@ -119,7 +119,8 @@ bool FPCGPrimitiveFactoryFastGeoPISMC::Create(FPCGContext* InContext)
 
 	FastGeo->AddComponentCluster(ComponentCluster.Get());
 
-	FastGeo->OnCreated();
+	// Disable the default reference collection path as it is currently too costly for runtime use. Instead manually collect references specific to the PISMC below.
+	FastGeo->OnCreated(/*bCollectReferences=*/false);
 
 	FastGeo->PrecachePSOs();
 
@@ -144,7 +145,8 @@ bool FPCGPrimitiveFactoryFastGeoPISMC::Create(FPCGContext* InContext)
 	});
 
 	UPCGManagedFastGeoContainer* ManagedPrimitives = FPCGContext::NewObject_AnyThread<UPCGManagedFastGeoContainer>(InContext, Cast<UObject>(ExecutionSource));
-	ManagedPrimitives->FastGeo = FastGeo;
+	ManagedPrimitives->SetFastGeoContainer(FastGeo);
+	ManagedPrimitives->SetObjectReferences(CollectObjectReferences());
 
 	ExecutionSource->GetExecutionState().AddToManagedResources(ManagedPrimitives);
 
@@ -155,3 +157,23 @@ FBox FPCGPrimitiveFactoryFastGeoPISMC::GetMeshBounds(int32 InPrimitiveIndex) con
 {
 	return ensure(MeshBounds.IsValidIndex(InPrimitiveIndex)) ? MeshBounds[InPrimitiveIndex] : FBox();
 }
+
+TArray<TObjectPtr<UObject>> FPCGPrimitiveFactoryFastGeoPISMC::CollectObjectReferences()
+{
+	TArray<TObjectPtr<UObject>> ObjectReferences;
+	ObjectReferences.Reserve(Descriptors.Num() * 4); // Heuristic
+
+	for (const FPCGProceduralISMComponentDescriptor& Descriptor : Descriptors)
+	{
+		ObjectReferences.Add(Descriptor.StaticMesh.Get());
+
+		if (Descriptor.OverlayMaterial)
+		{
+			ObjectReferences.Add(Descriptor.OverlayMaterial);
+		}
+
+		ObjectReferences.Append(Descriptor.OverrideMaterials);
+	}
+
+	return ObjectReferences;
+}
diff --git a/Engine/Plugins/Experimental/PCGInterops/PCGFastGeoInterop/Source/PCGFastGeoInterop/Private/Compute/PrimitiveFactories/PCGPrimitiveFactoryFastGeoPISMC.h b/Engine/Plugins/Experimental/PCGInterops/PCGFastGeoInterop/Source/PCGFastGeoInterop/Private/Compute/PrimitiveFactories/PCGPrimitiveFactoryFastGeoPISMC.h
--- a/Engine/Plugins/Experimental/PCGInterops/PCGFastGeoInterop/Source/PCGFastGeoInterop/Private/Compute/PrimitiveFactories/PCGPrimitiveFactoryFastGeoPISMC.h
+++ b/Engine/Plugins/Experimental/PCGInterops/PCGFastGeoInterop/Source/PCGFastGeoInterop/Private/Compute/PrimitiveFactories/PCGPrimitiveFactoryFastGeoPISMC.h
@@ -28,6 +28,9 @@ public:
 	virtual FBox GetMeshBounds(int32 InPrimitiveIndex) const override;
 	//~ End IPCGPrimitiveFactoryISMBase interface
 
+protected:
+	TArray<TObjectPtr<UObject>> CollectObjectReferences();
+
 protected:
 	TArray<FPCGProceduralISMComponentDescriptor> Descriptors;
 
