[FastGeo] Add runtime API for PCG componentless scattering
- Add CreateRuntime/DestroyRuntime/IsRuntime API replacing CreateTransient/DestroyTransient.
- Add OnRegistered delegate so consumers can subscribe instead of polling HasAnyPendingTasks/Tick.
- Subsystem drives runtime container registration and destruction via PendingRuntimeRegisterContainers/PendingRuntimeDestroyContainers.
- PCG uses callback-based CreateRuntime with init cluster pattern, removing the 5-step creation ceremony.
- PCG no longer drives Container->Tick() from IsRenderStateCreated; subsystem handles it.
- Move OnCreated to private; only CreateRuntime and the world partition transformer call it.

#Synced-CL 52208678

diff --git a/Engine/Plugins/Experimental/FastGeoStreaming/Source/FastGeoStreaming/Internal/FastGeoContainer.h b/Engine/Plugins/Experimental/FastGeoStreaming/Source/FastGeoStreaming/Internal/FastGeoContainer.h
--- a/Engine/Plugins/Experimental/FastGeoStreaming/Source/FastGeoStreaming/Internal/FastGeoContainer.h
+++ b/Engine/Plugins/Experimental/FastGeoStreaming/Source/FastGeoStreaming/Internal/FastGeoContainer.h
@@ -32,6 +32,16 @@ class FFastGeoComponentCluster;
 class FChaosUserDefinedEntity;
 class AFastGeoSurrogateActor;
 
+/** Result of UFastGeoContainer::CreateRuntime. */
+struct FFastGeoCreateRuntimeResult
+{
+	/** The created container, or null on failure. */
+	UFastGeoContainer* Container = nullptr;
+
+	/** Component pointers built from the container's internal component iteration order. */
+	TArray<FFastGeoComponent*> Components;
+};
+
 /**
  * Cached world/scene/subsystem/physics pointers.
  * Populated on GT before workers run.
@@ -186,14 +196,29 @@ class FASTGEOSTREAMING_API UFastGeoContainer : public UAssetUserData, public IAs
 public:
 
 	/**
-	 * Creates a transient FastGeoContainer for a specific world (used for runtime).
+	 * Creates a runtime FastGeoContainer and starts async registration.
+	 * The subsystem drives registration ticks automatically.
+	 * Subscribe to GetOnRegistered() for completion notification, or check IsFullyRegistered().
+	 *
+	 * @param InWorld				World in which to create the container.
+	 * @param InDebugName			Debug name for the container and its component cluster.
+	 * @param InInitCluster			Callback to populate the component cluster with components.
+	 * @param bInCollectReferences	Whether to collect asset references via serialization archive. Pass false if the caller manages asset lifetimes externally.
+	 *								
+	 * @return Result containing the container and an ordered array of component pointers matching the components added in InInitCluster. Container is null on failure.
 	 */
-	static UFastGeoContainer* CreateTransient(UWorld* InWorld);
+	static FFastGeoCreateRuntimeResult CreateRuntime(UWorld* InWorld, FName InDebugName, TFunctionRef<void(FFastGeoComponentCluster&)> InInitCluster, bool bInCollectReferences = true);
 
 	/**
-	 * Destroys a transient FastGeoContainer.
+	 * Destroys a runtime FastGeoContainer created by CreateRuntime.
+	 * Initiates async unregistration; the subsystem drives completion.
+	 *
+	 * @param InFastGeo		The runtime container to destroy. Must satisfy IsRuntime().
 	 */
-	static void DestroyTransient(UFastGeoContainer* InFastGeo);
+	static void DestroyRuntime(UFastGeoContainer* InFastGeo);
+
+	/** Returns true if this is a runtime (transient) container. Only CreateRuntime sets RF_Transient on FastGeo containers. */
+	bool IsRuntime() const { return HasAnyFlags(RF_Transient); }
 
 	void Register();
 	void Unregister();
@@ -231,7 +256,6 @@ public:
 	bool IsUsingSurrogateComponents() const;
 	AFastGeoSurrogateActor* GetSurrogateActor() const;
 
-	void OnCreated(bool bCollectReferences = true);
 	void AddComponentCluster(FFastGeoComponentCluster* ComponentCluster);
 
 	void InitializeDynamicProperties(bool bInitForPlay = false);
@@ -289,6 +313,9 @@ public:
 		return ForEachComponentClusterBreakable<UFastGeoContainer, TComponentCluster, TFunc>(this, Forward<TFunc>(InFunc));
 	}
 
+	DECLARE_MULTICAST_DELEGATE(FOnRegistered);
+	FOnRegistered& GetOnRegistered() { return OnRegistered; }
+
 	DECLARE_MULTICAST_DELEGATE(FOnUnregistered);
 	FOnUnregistered& GetOnUnregistered() { return OnUnregistered; }
 
@@ -304,6 +331,7 @@ private:
 	FPhysScene* GetPhysicsScene() const;
 	float GetWorldTimeSeconds() const;
 
+	void OnCreated(bool bCollectReferences = true);
 	void InitializePrerequisites();
 	void ForEachMaterialAsset(TFunctionRef<void(UMaterialInterface*)> Func);
 	void BuildMaterialAssets();
@@ -490,6 +518,7 @@ private:
 
 	FRenderCommandFence DestroyFence;
 
+	FOnRegistered OnRegistered;
 	FOnUnregistered OnUnregistered;
 
 	friend class UFastGeoWorldSubsystem;
diff --git a/Engine/Plugins/Experimental/FastGeoStreaming/Source/FastGeoStreaming/Internal/FastGeoWorldSubsystem.h b/Engine/Plugins/Experimental/FastGeoStreaming/Source/FastGeoStreaming/Internal/FastGeoWorldSubsystem.h
--- a/Engine/Plugins/Experimental/FastGeoStreaming/Source/FastGeoStreaming/Internal/FastGeoWorldSubsystem.h
+++ b/Engine/Plugins/Experimental/FastGeoStreaming/Source/FastGeoStreaming/Internal/FastGeoWorldSubsystem.h
@@ -118,13 +118,17 @@ private:
 	void ForEachHLODObjectInCell(const UWorldPartitionRuntimeCell*, TFunction<void(IWorldPartitionHLODObject*)>);
 	void RequestAsyncRenderStateTasksBudget_Concurrent(float& OutAvailableTimeBudgetMS, int32& OutAvaiableComponentsBudget, int32& OutTimeEpoch);
 	void CommitAsyncRenderStateTasksBudget_Concurrent(float InUsedTimeBudgetMS, int32& InUsedComponentsBudget, int32 TimeEpoch);
-	void PushAsyncDestroyContainer(UFastGeoContainer* FastGeo);
 	FFastGeoAsyncRenderStateJobQueue* GetOrCreateJobQueue();
 	void WaitForAllPendingWorkCompletion();
-	void DrainPendingDestroyContainers();
+	void PushRuntimeContainerForRegistration(UFastGeoContainer* FastGeo);
+	void PushRuntimeContainerForDestruction(UFastGeoContainer* FastGeo);
+	void DrainPendingRuntimeContainers();
 
 	UPROPERTY()
-	TSet<TObjectPtr<UFastGeoContainer>> PendingDestroyContainers;
+	TSet<TObjectPtr<UFastGeoContainer>> PendingRuntimeRegisterContainers;
+
+	UPROPERTY()
+	TSet<TObjectPtr<UFastGeoContainer>> PendingRuntimeDestroyContainers;
 
 	FDelegateHandle Handle_OnLevelStreamingStateChanged;
 	FDelegateHandle Handle_OnLevelBeginAddToWorld;
diff --git a/Engine/Plugins/Experimental/FastGeoStreaming/Source/FastGeoStreaming/Private/FastGeoContainer.cpp b/Engine/Plugins/Experimental/FastGeoStreaming/Source/FastGeoStreaming/Private/FastGeoContainer.cpp
--- a/Engine/Plugins/Experimental/FastGeoStreaming/Source/FastGeoStreaming/Private/FastGeoContainer.cpp
+++ b/Engine/Plugins/Experimental/FastGeoStreaming/Source/FastGeoStreaming/Private/FastGeoContainer.cpp
@@ -370,6 +370,13 @@ void UFastGeoContainer::OnRegisterCompleted()
 	check(!HasAnyPendingCreateTasks());
 	check(RegistrationState == ERegistrationState::Registering);
 	RegistrationState = ERegistrationState::Registered;
+
+	// Only broadcast if we're not about to immediately unregister (e.g., DestroyRuntime called during registration).
+	if (RegistrationTargetState == ERegistrationTargetState::Registered)
+	{
+		OnRegistered.Broadcast();
+	}
+
 	ApplyRegistrationTargetState();
 }
 
@@ -388,60 +395,92 @@ void UFastGeoContainer::OnUnregisterCompleted()
 	ApplyRegistrationTargetState();
 }
 
-UFastGeoContainer* UFastGeoContainer::CreateTransient(UWorld* InWorld)
+FFastGeoCreateRuntimeResult UFastGeoContainer::CreateRuntime(UWorld* InWorld, FName InDebugName, TFunctionRef<void(FFastGeoComponentCluster&)> InInitCluster, bool bInCollectReferences)
 {
+	TRACE_CPUPROFILER_EVENT_SCOPE(UFastGeoContainer::CreateRuntime);
 	check(IsInGameThread());
 
+	FFastGeoCreateRuntimeResult Result;
+
+	if (!FFastGeoStreamingModule::IsFastGeoEnabled())
+	{
+		UE_LOG(LogFastGeoStreaming, Warning, TEXT("UFastGeoContainer::CreateRuntime() ignored because FastGeo is disabled."));
+		return Result;
+	}
+
 	if (!IsValid(InWorld))
 	{
-		UE_LOG(LogFastGeoStreaming, Error, TEXT("UFastGeoContainer::CreateTransient() called with an invalid world."));
-		return nullptr;
+		UE_LOG(LogFastGeoStreaming, Error, TEXT("UFastGeoContainer::CreateRuntime() called with an invalid world."));
+		return Result;
 	}
 
 	ULevel* PersistentLevel = InWorld->PersistentLevel;
 	check(PersistentLevel && PersistentLevel->IsPersistentLevel());
 
-	return NewObject<UFastGeoContainer>(PersistentLevel, NAME_None, RF_Transient);
+	const FName UniqueName = MakeUniqueObjectName(PersistentLevel, UFastGeoContainer::StaticClass(), InDebugName);
+	UFastGeoContainer* Container = NewObject<UFastGeoContainer>(PersistentLevel, UniqueName, RF_Transient);
+
+	// Create the cluster, let the caller populate it with components, then add to container.
+	FFastGeoComponentCluster TempCluster(Container, *FString::Printf(TEXT("FastGeoComponentCluster_%s"), *InDebugName.ToString()));
+	InInitCluster(TempCluster);
+	Container->AddComponentCluster(&TempCluster);
+
+	Container->OnCreated(bInCollectReferences);
+	Container->PrecachePSOs();
+	Container->Register();
+	Container->GetWorldSubsystem()->PushRuntimeContainerForRegistration(Container);
+
+	// Populate result.
+	Result.Container = Container;
+	Result.Container->ForEachComponentCluster([&Result](FFastGeoComponentCluster& InComponentCluster)
+	{
+		InComponentCluster.ForEachComponent([&Result](FFastGeoComponent& InComponent)
+		{
+			Result.Components.Add(&InComponent);
+		});
+	});
+
+	return Result;
 }
 
-void UFastGeoContainer::DestroyTransient(UFastGeoContainer* InFastGeo)
+void UFastGeoContainer::DestroyRuntime(UFastGeoContainer* InFastGeo)
 {
+	TRACE_CPUPROFILER_EVENT_SCOPE(UFastGeoContainer::DestroyRuntime);
 	check(IsInGameThread());
 
 	if (!IsValid(InFastGeo))
 	{
-		UE_LOG(LogFastGeoStreaming, Error, TEXT("UFastGeoContainer::DestroyTransient() called on an invalid container."));
+		UE_LOG(LogFastGeoStreaming, Error, TEXT("UFastGeoContainer::DestroyRuntime() called on an invalid container."));
 		return;
 	}
 
 	if (!InFastGeo->GetWorld())
 	{
-		UE_LOG(LogFastGeoStreaming, Error, TEXT("UFastGeoContainer::DestroyTransient() called on a container with an invalid world."));
+		UE_LOG(LogFastGeoStreaming, Error, TEXT("UFastGeoContainer::DestroyRuntime() called on a container with an invalid world."));
 		return;
 	}
 
-	if (!InFastGeo->HasAnyFlags(RF_Transient))
+	if (!InFastGeo->IsRuntime())
 	{
-		UE_LOG(LogFastGeoStreaming, Error, TEXT("UFastGeoContainer::DestroyTransient() is meant to be called on a transient container."));
+		UE_LOG(LogFastGeoStreaming, Error, TEXT("UFastGeoContainer::DestroyRuntime() is meant to be called on a runtime container."));
 		return;
 	}
 
 	if (!InFastGeo->GetLevel()->IsPersistentLevel())
 	{
-		UE_LOG(LogFastGeoStreaming, Error, TEXT("UFastGeoContainer::DestroyTransient() expects the transient container to be outered to the persistent level."));
+		UE_LOG(LogFastGeoStreaming, Error, TEXT("UFastGeoContainer::DestroyRuntime() expects the runtime container to be outered to the persistent level."));
 		return;
 	}
 
 	if (InFastGeo->bIsBeingAsyncDestroyed)
 	{
-		UE_LOG(LogFastGeoStreaming, Error, TEXT("UFastGeoContainer::DestroyTransient() called on a container that is already being destroyed."));
+		UE_LOG(LogFastGeoStreaming, Error, TEXT("UFastGeoContainer::DestroyRuntime() called on a container that is already being destroyed."));
 		return;
 	}
 
 	InFastGeo->bIsBeingAsyncDestroyed = true;
-
-	UFastGeoWorldSubsystem* WorldSubsystem = InFastGeo->GetWorldSubsystem();
-	WorldSubsystem->PushAsyncDestroyContainer(InFastGeo);
+	InFastGeo->Unregister();
+	InFastGeo->GetWorldSubsystem()->PushRuntimeContainerForDestruction(InFastGeo);
 }
 
 void UFastGeoContainer::Unregister()
diff --git a/Engine/Plugins/Experimental/FastGeoStreaming/Source/FastGeoStreaming/Private/FastGeoWorldSubsystem.cpp b/Engine/Plugins/Experimental/FastGeoStreaming/Source/FastGeoStreaming/Private/FastGeoWorldSubsystem.cpp
--- a/Engine/Plugins/Experimental/FastGeoStreaming/Source/FastGeoStreaming/Private/FastGeoWorldSubsystem.cpp
+++ b/Engine/Plugins/Experimental/FastGeoStreaming/Source/FastGeoStreaming/Private/FastGeoWorldSubsystem.cpp
@@ -313,19 +313,38 @@ void UFastGeoWorldSubsystem::OnWorldCleanup(UWorld* World, bool bSessionEnded, b
 {
 	if (World == GetWorld())
 	{
-		DrainPendingDestroyContainers();
+		DrainPendingRuntimeContainers();
 	}
 }
 
-void UFastGeoWorldSubsystem::DrainPendingDestroyContainers()
+void UFastGeoWorldSubsystem::DrainPendingRuntimeContainers()
 {
 	check(IsInGameThread());
 
 	static constexpr bool bWaitForCompletion = true;
 
-	while (!PendingDestroyContainers.IsEmpty())
+	// Force-complete registering containers, then unregister them.
+	// Note: Unregister() is safe to call even if an OnRegistered callback already called
+	// DestroyRuntime -> Unregister; the state machine handles redundant calls as no-ops.
+	while (!PendingRuntimeRegisterContainers.IsEmpty())
 	{
-		TSet<TObjectPtr<UFastGeoContainer>> Containers = MoveTemp(PendingDestroyContainers);
+		TSet<TObjectPtr<UFastGeoContainer>> Containers = MoveTemp(PendingRuntimeRegisterContainers);
+		for (const TObjectPtr<UFastGeoContainer>& Container : Containers)
+		{
+			if (IsValid(Container))
+			{
+				Container->Tick(bWaitForCompletion);
+				Container->Unregister();
+				Container->Tick(bWaitForCompletion);
+				check(Container->IsFullyUnregistered());
+			}
+		}
+	}
+
+	// Force-complete destroying containers.
+	while (!PendingRuntimeDestroyContainers.IsEmpty())
+	{
+		TSet<TObjectPtr<UFastGeoContainer>> Containers = MoveTemp(PendingRuntimeDestroyContainers);
 		for (const TObjectPtr<UFastGeoContainer>& Container : Containers)
 		{
 			if (IsValid(Container))
@@ -335,7 +354,9 @@ void UFastGeoWorldSubsystem::DrainPendingDestroyContainers()
 			}
 		}
 	}
-	check(PendingDestroyContainers.IsEmpty());
+
+	check(PendingRuntimeRegisterContainers.IsEmpty());
+	check(PendingRuntimeDestroyContainers.IsEmpty());
 }
 
 void UFastGeoWorldSubsystem::OnGlobalComponentReregisterContextCreated()
@@ -531,13 +552,18 @@ void UFastGeoWorldSubsystem::WaitForAllPendingWorkCompletion()
 #endif
 	}
 
-	DrainPendingDestroyContainers();
+	// The TObjectIterator loop above already ticked all containers to completion.
+	// Clear the register set before draining -- don't force-unregister, as this may be called
+	// during FGlobalComponentReregisterContext where containers must survive.
+	PendingRuntimeRegisterContainers.Empty();
+	DrainPendingRuntimeContainers();
 
 	FPhysScene* PhysScene = World ? World->GetPhysicsScene() : nullptr;
 	check(!PhysScene || !PhysScene->HasAsyncPhysicsStateJobs());
 	check(!AsyncRenderStateJobQueue.IsValid() || AsyncRenderStateJobQueue->IsCompleted());
 	check(ComponentsPendingRecreate.IsEmpty());
-	check(PendingDestroyContainers.IsEmpty());
+	check(PendingRuntimeRegisterContainers.IsEmpty());
+	check(PendingRuntimeDestroyContainers.IsEmpty());
 }
 
 void UFastGeoWorldSubsystem::AddToComponentsPendingRecreate(FFastGeoComponent* InComponentPendingRecreate)
@@ -654,21 +680,33 @@ void UFastGeoWorldSubsystem::Tick(float DeltaTime)
 		UsedAsyncRenderStateTasksTimeBudgetMS = 0;
 		UsedNumComponentsToProcessBudget = 0;
 	}
-	
-	for (auto It = PendingDestroyContainers.CreateIterator(); It; ++It)
+
+	// Tick runtime containers and remove completed ones.
+	auto TickAndDrainCompleted = [](TSet<TObjectPtr<UFastGeoContainer>>& Containers, auto IsCompletedFunc)
 	{
-		UFastGeoContainer* Container = *It;
-		bool bShouldRemove = !IsValid(Container) || Container->IsFullyUnregistered();
-		if (!bShouldRemove)
-		{
-			Container->Tick();
-			bShouldRemove = Container->IsFullyUnregistered();
-		}
-		if (bShouldRemove)
+		for (auto It = Containers.CreateIterator(); It; ++It)
 		{
-			It.RemoveCurrent();
+			UFastGeoContainer* Container = *It;
+			bool bShouldRemove = !IsValid(Container) || IsCompletedFunc(Container);
+			if (!bShouldRemove)
+			{
+				Container->Tick();
+				bShouldRemove = IsCompletedFunc(Container);
+			}
+			if (bShouldRemove)
+			{
+				It.RemoveCurrent();
+			}
 		}
-	}
+	};
+
+	// Level-streamed containers are driven by OnAddLevelToWorldExtension/OnRemoveLevelFromWorldExtension instead.
+	// Also evict containers marked for destruction (e.g., DestroyRuntime called from an OnRegistered callback).
+	// These are tracked by PendingRuntimeDestroyContainers and should not be ticked from both sets.
+	TickAndDrainCompleted(PendingRuntimeRegisterContainers, [](const UFastGeoContainer* Container) { return Container->IsFullyRegistered() || Container->bIsBeingAsyncDestroyed; });
+
+	// Containers complete destruction when fully unregistered.
+	TickAndDrainCompleted(PendingRuntimeDestroyContainers, [](const UFastGeoContainer* Container) { return Container->IsFullyUnregistered(); });
 }
 
 TStatId UFastGeoWorldSubsystem::GetStatId() const
@@ -721,16 +759,31 @@ void UFastGeoWorldSubsystem::CommitAsyncRenderStateTasksBudget_Concurrent(float
 	}
 }
 
-void UFastGeoWorldSubsystem::PushAsyncDestroyContainer(UFastGeoContainer* FastGeo)
+void UFastGeoWorldSubsystem::PushRuntimeContainerForRegistration(UFastGeoContainer* FastGeo)
+{
+	check(IsInGameThread());
+	check(FastGeo->IsRuntime());
+	check(!FastGeo->bIsBeingAsyncDestroyed);
+
+	if (!FastGeo->IsFullyRegistered())
+	{
+		PendingRuntimeRegisterContainers.Add(FastGeo);
+	}
+}
+
+void UFastGeoWorldSubsystem::PushRuntimeContainerForDestruction(UFastGeoContainer* FastGeo)
 {
 	check(IsInGameThread());
+	check(FastGeo->IsRuntime());
 	check(FastGeo->bIsBeingAsyncDestroyed);
 
-	FastGeo->Unregister();
+	// Don't remove from PendingRuntimeRegisterContainers here -- the register
+	// tick loop checks bIsBeingAsyncDestroyed and removes it. This avoids
+	// invalidating iterators if DestroyRuntime is called from an OnRegistered callback.
 
 	if (!FastGeo->IsFullyUnregistered())
 	{
-		PendingDestroyContainers.Add(FastGeo);
+		PendingRuntimeDestroyContainers.Add(FastGeo);
 	}
 }
 
diff --git a/Engine/Plugins/Experimental/PCGInterops/PCGFastGeoInterop/Source/PCGFastGeoInterop/Private/Components/PCGManagedFastGeoContainer.cpp b/Engine/Plugins/Experimental/PCGInterops/PCGFastGeoInterop/Source/PCGFastGeoInterop/Private/Components/PCGManagedFastGeoContainer.cpp
--- a/Engine/Plugins/Experimental/PCGInterops/PCGFastGeoInterop/Source/PCGFastGeoInterop/Private/Components/PCGManagedFastGeoContainer.cpp
+++ b/Engine/Plugins/Experimental/PCGInterops/PCGFastGeoInterop/Source/PCGFastGeoInterop/Private/Components/PCGManagedFastGeoContainer.cpp
@@ -10,7 +10,7 @@ bool UPCGManagedFastGeoContainer::Release(bool bHardRelease, TSet<TSoftObjectPtr
 
 	if (FastGeo)
 	{
-		UFastGeoContainer::DestroyTransient(FastGeo);
+		UFastGeoContainer::DestroyRuntime(FastGeo);
 		FastGeo = nullptr;
 	}
 
diff --git a/Engine/Plugins/Experimental/PCGInterops/PCGFastGeoInterop/Source/PCGFastGeoInterop/Private/Compute/PrimitiveFactories/PCGPrimitiveFactoryFastGeoPISMC.cpp b/Engine/Plugins/Experimental/PCGInterops/PCGFastGeoInterop/Source/PCGFastGeoInterop/Private/Compute/PrimitiveFactories/PCGPrimitiveFactoryFastGeoPISMC.cpp
--- a/Engine/Plugins/Experimental/PCGInterops/PCGFastGeoInterop/Source/PCGFastGeoInterop/Private/Compute/PrimitiveFactories/PCGPrimitiveFactoryFastGeoPISMC.cpp
+++ b/Engine/Plugins/Experimental/PCGInterops/PCGFastGeoInterop/Source/PCGFastGeoInterop/Private/Compute/PrimitiveFactories/PCGPrimitiveFactoryFastGeoPISMC.cpp
@@ -21,23 +21,10 @@ bool FPCGPrimitiveFactoryFastGeoPISMC::IsRenderStateCreated() const
 		return true;
 	}
 
-	// Use first component to get to container.
-	FFastGeoComponent* Component = Components[0].Get();
-	UFastGeoContainer* Container = Component ? Component->GetOwnerContainer() : nullptr;
-
-	if (!Container)
-	{
-		return true;
-	}
-
-	if (Container->HasAnyPendingTasks())
+	// Registration is driven by the subsystem; wait for the OnRegistered callback.
+	if (!*bFastGeoRegistered)
 	{
-		Container->Tick();
-
-		if (Container->HasAnyPendingTasks())
-		{
-			return false;
-		}
+		return false;
 	}
 
 	for (int32 PrimitiveIndex = 0; PrimitiveIndex < GetNumPrimitives(); ++PrimitiveIndex)
@@ -103,12 +90,6 @@ void FPCGPrimitiveFactoryFastGeoPISMC::Initialize(FParameters&& InParameters)
 
 bool FPCGPrimitiveFactoryFastGeoPISMC::Create(FPCGContext* InContext)
 {
-	if (!FFastGeoStreamingModule::IsFastGeoEnabled())
-	{
-		UE_LOG(LogPCG, Warning, TEXT("FPCGPrimitiveFactoryFastGeoPISMC: FastGeoStreaming plugin is not enabled, no primitives will be spawned."));
-		return true;
-	}
-
 	IPCGGraphExecutionSource* ExecutionSource = InContext->ExecutionSource.Get();
 	if (!ExecutionSource)
 	{
@@ -117,9 +98,9 @@ bool FPCGPrimitiveFactoryFastGeoPISMC::Create(FPCGContext* InContext)
 
 	UWorld* World = ExecutionSource->GetExecutionState().GetWorld();
 
+	AActor* HitProxyTargetActor = nullptr;
 #if WITH_EDITOR
-	AActor* HitProxyTargetActor = ExecutionSource->GetExecutionState().GetTypedTarget<AActor>();
-
+	HitProxyTargetActor = ExecutionSource->GetExecutionState().GetTypedTarget<AActor>();
 	// For actor-component-less generation the local execution source has no target actor.
 	// The original source (e.g. the UPCGComponent on the user's actor) typically does.
 	if (!HitProxyTargetActor)
@@ -131,77 +112,62 @@ bool FPCGPrimitiveFactoryFastGeoPISMC::Create(FPCGContext* InContext)
 	}
 #endif // WITH_EDITOR
 
-	UFastGeoContainer* FastGeo = UFastGeoContainer::CreateTransient(World);
-	if (!FastGeo)
-	{
-		UE_LOG(LogPCG, Error, TEXT("Failed to create a transient FastGeoContainer, geometry will not be created."));
-		return true;
-	}
-
 	const FString DebugName = ExecutionSource->GetExecutionState().GetDebugName();
-	TUniquePtr<FFastGeoComponentCluster> ComponentCluster = MakeUnique<FFastGeoComponentCluster>(FastGeo, *FString::Printf(TEXT("FastGeoComponentCluster_%s"), *DebugName));
-
 	const int32 NumToCreate = FMath::Min(PrimitiveInfos.Num(), PCGComputeConstants::MAX_PRIMITIVE_COMPONENTS_PER_SPAWNER);
 
 	MeshBounds.SetNumUninitialized(NumToCreate);
 
-	for (int32 DescIndex = 0; DescIndex < NumToCreate; ++DescIndex)
+	FFastGeoCreateRuntimeResult Result = UFastGeoContainer::CreateRuntime(World, *FString::Printf(TEXT("PCG_%s"), *DebugName), [this, NumToCreate, HitProxyTargetActor](FFastGeoComponentCluster& InCluster)
 	{
-		FPCGPrimitiveInfo& PrimitiveInfo = PrimitiveInfos[DescIndex];
-		FPCGProceduralISMComponentDescriptor& Desc = PrimitiveInfo.Descriptor;
+		for (int32 DescIndex = 0; DescIndex < NumToCreate; ++DescIndex)
+		{
+			FPCGPrimitiveInfo& PrimitiveInfo = PrimitiveInfos[DescIndex];
+			FPCGProceduralISMComponentDescriptor& Desc = PrimitiveInfo.Descriptor;
 
-		TArray<FInstanceSceneDataBuffers::FCompressedSpatialHashItem> SpatialHashes;
-		FBox BoundsOfAllCells(EForceInit::ForceInit);
+			TArray<FInstanceSceneDataBuffers::FCompressedSpatialHashItem> SpatialHashes;
+			FBox BoundsOfAllCells(EForceInit::ForceInit);
 
-		PopulateHashes(PrimitiveInfo, SpatialHashes, BoundsOfAllCells);
+			PopulateHashes(PrimitiveInfo, SpatialHashes, BoundsOfAllCells);
 
-		if (BoundsOfAllCells.IsValid)
-		{
-			Desc.WorldBounds = Desc.WorldBounds.Overlap(BoundsOfAllCells);
-		}
+			if (BoundsOfAllCells.IsValid)
+			{
+				Desc.WorldBounds = Desc.WorldBounds.Overlap(BoundsOfAllCells);
+			}
 
-		FFastGeoProceduralISMComponent& FastGeoComponent = static_cast<FFastGeoProceduralISMComponent&>(ComponentCluster->AddComponent(FFastGeoProceduralISMComponent::Type));
-		FastGeoComponent.InitializeFromComponentDescriptor(Desc);
-		FastGeoComponent.SetCustomPrimitiveData(CustomPrimitiveData.GetViewForIndex(DescIndex));
-		FastGeoComponent.SetSpatialHashes(MoveTemp(SpatialHashes));
+			FFastGeoProceduralISMComponent& FastGeoComponent = static_cast<FFastGeoProceduralISMComponent&>(InCluster.AddComponent(FFastGeoProceduralISMComponent::Type));
+			FastGeoComponent.InitializeFromComponentDescriptor(Desc);
+			FastGeoComponent.SetCustomPrimitiveData(CustomPrimitiveData.GetViewForIndex(DescIndex));
+			FastGeoComponent.SetSpatialHashes(MoveTemp(SpatialHashes));
 #if WITH_EDITOR
-		FastGeoComponent.SetHitProxyTargetActor(HitProxyTargetActor);
+			FastGeoComponent.SetHitProxyTargetActor(HitProxyTargetActor);
 #endif
+			MeshBounds[DescIndex] = Desc.StaticMesh->GetBoundingBox();
+		}
 
-		MeshBounds[DescIndex] = Desc.StaticMesh->GetBoundingBox();
-	}
-
-	ensure(ComponentCluster->HasComponents());
-
-	FastGeo->AddComponentCluster(ComponentCluster.Get());
-
-	// Disable the default reference collection path as it is currently too costly for runtime use. Instead manually collect references specific to the PISMC below.
-	FastGeo->OnCreated(/*bCollectReferences=*/false);
-
-	FastGeo->PrecachePSOs();
-
-	FastGeo->Register();
+		ensure(InCluster.HasComponents());
+	}, /*bCollectReferences=*/false); // Disable the default reference collection path as it is currently too costly for runtime use. Instead manually collect references specific to the PISMC below.
 
-	int32 ComponentIndex = 0;
-	FastGeo->ForEachComponentCluster([This=this, &ComponentIndex](FFastGeoComponentCluster& InComponentCluster)
+	if (!Result.Container)
 	{
-		InComponentCluster.ForEachComponent([This, &ComponentIndex](FFastGeoComponent& InComponent)
-		{
-			check(This->PrimitiveInfos.IsValidIndex(ComponentIndex));
-			const FPCGProceduralISMComponentDescriptor& Desc = This->PrimitiveInfos[ComponentIndex].Descriptor;
-
-			// If not true, weakptr assignment below will fail.
-			check(InComponent.GetOwnerComponentCluster());
+		UE_LOG(LogPCG, Error, TEXT("Failed to create a runtime FastGeoContainer, geometry will not be created."));
+		return true;
+	}
 
-			This->Components.Add(&InComponent);
-			This->InstanceCounts.Add(Desc.NumInstances);
+	*bFastGeoRegistered = Result.Container->IsFullyRegistered();
+	if (!*bFastGeoRegistered)
+	{
+		Result.Container->GetOnRegistered().AddLambda([Flag = bFastGeoRegistered]() { *Flag = true; });
+	}
 
-			++ComponentIndex;
-		});
-	});
+	for (int32 Index = 0; Index < Result.Components.Num(); ++Index)
+	{
+		check(PrimitiveInfos.IsValidIndex(Index));
+		Components.Add(Result.Components[Index]);
+		InstanceCounts.Add(PrimitiveInfos[Index].Descriptor.NumInstances);
+	}
 
 	UPCGManagedFastGeoContainer* ManagedPrimitives = FPCGContext::NewObject_AnyThread<UPCGManagedFastGeoContainer>(InContext, Cast<UObject>(ExecutionSource));
-	ManagedPrimitives->SetFastGeoContainer(FastGeo);
+	ManagedPrimitives->SetFastGeoContainer(Result.Container);
 	ManagedPrimitives->SetObjectReferences(CollectObjectReferences());
 
 	FPCGManagedResourceContainerHelper ContainerHelper(ExecutionSource);
diff --git a/Engine/Plugins/Experimental/PCGInterops/PCGFastGeoInterop/Source/PCGFastGeoInterop/Private/Compute/PrimitiveFactories/PCGPrimitiveFactoryFastGeoPISMC.h b/Engine/Plugins/Experimental/PCGInterops/PCGFastGeoInterop/Source/PCGFastGeoInterop/Private/Compute/PrimitiveFactories/PCGPrimitiveFactoryFastGeoPISMC.h
--- a/Engine/Plugins/Experimental/PCGInterops/PCGFastGeoInterop/Source/PCGFastGeoInterop/Private/Compute/PrimitiveFactories/PCGPrimitiveFactoryFastGeoPISMC.h
+++ b/Engine/Plugins/Experimental/PCGInterops/PCGFastGeoInterop/Source/PCGFastGeoInterop/Private/Compute/PrimitiveFactories/PCGPrimitiveFactoryFastGeoPISMC.h
@@ -40,4 +40,8 @@ protected:
 	TArray<FWeakFastGeoComponent> Components;
 	TArray<int32> InstanceCounts;
 	TArray<FBox> MeshBounds;
+
+	// Shared so the OnRegistered lambda can safely write to it even if this factory is destroyed
+	// before the container finishes registration.
+	TSharedRef<bool> bFastGeoRegistered = MakeShared<bool>(false);
 };
